APT: Remove Editor and Engine dependencies from APT Launcher Extension

* Move Launch extension to a separate plugin to remove any Engine/Editor dependencies
* APT can now be used to run Automated perf tests with Unreal Frontend
* Refactor code to read config from ini instead of Developer Setting UObjects
* Add a few config options to specify LLM, GPU Perf etc.

#rnx
#tests Launch APT replay, sequence tests on Lyra via Project Launcher in-editor as well as Unreal Frontend
#Synced-CL 50456601

diff --git a/Engine/Plugins/Performance/AutomatedPerfTestTools/AutomatedPerfTestTools.uplugin b/Engine/Plugins/Performance/AutomatedPerfTestTools/AutomatedPerfTestTools.uplugin
new file mode 100644
--- /dev/null
+++ b/Engine/Plugins/Performance/AutomatedPerfTestTools/AutomatedPerfTestTools.uplugin
@@ -0,0 +1,43 @@
+{
+	"FileVersion": 3,
+	"Version": 1,
+	"VersionName": "0.1",
+	"FriendlyName": "AutomatedPerfTestTools",
+	"Description": "Tools for Automated Perf Testing framework",
+	"Category": "Testing",
+	"CreatedBy": "",
+	"CreatedByURL": "",
+	"DocsURL": "",
+	"MarketplaceURL": "",
+	"SupportURL": "",
+	"CanContainContent": true,
+	"IsBetaVersion": false,
+	"IsExperimentalVersion": true,
+	"Installed": true,
+	"EnabledByDefault": true,
+	"Modules": [
+		{
+			"Name": "AutomatedPerfTestLaunchExtension",
+			"Type": "Editor",
+			"LoadingPhase": "Default",
+			"PlatformAllowList": [
+				"Win64",
+				"Linux",
+				"Mac"
+			],
+			"ProgramAllowList": [ "UnrealFrontend" ]
+		}
+	],
+	"SupportedPrograms": [ "UnrealFrontend" ],
+	"Plugins": [
+		{
+			"Name": "ProjectLauncher",
+			"Enabled": true,
+			"SupportedTargetPlatforms": [
+				"Win64",
+				"Linux",
+				"Mac"
+			]
+		}
+	]
+}
\ No newline at end of file
diff --git a/Engine/Plugins/Performance/AutomatedPerfTestTools/Source/AutomatedPerfTestLaunchExtension/AutomatedPerfTestLaunchExtension.Build.cs b/Engine/Plugins/Performance/AutomatedPerfTestTools/Source/AutomatedPerfTestLaunchExtension/AutomatedPerfTestLaunchExtension.Build.cs
new file mode 100644
--- /dev/null
+++ b/Engine/Plugins/Performance/AutomatedPerfTestTools/Source/AutomatedPerfTestLaunchExtension/AutomatedPerfTestLaunchExtension.Build.cs
@@ -0,0 +1,33 @@
+// Copyright Epic Games, Inc. All Rights Reserved.
+
+using UnrealBuildTool;
+
+public class AutomatedPerfTestLaunchExtension : ModuleRules
+{
+	public AutomatedPerfTestLaunchExtension(ReadOnlyTargetRules Target) : base(Target)
+	{
+		PrivateDependencyModuleNames.AddRange(
+		[
+			"Core",
+			"CoreUObject",
+			"Json",
+			"Slate",
+			"SlateCore",
+			"Sockets",
+			"TraceLog",
+			"ProjectLauncher",
+			"DeveloperSettings",
+			"ToolWidgets",
+		]);
+
+
+		if (Target.bBuildEditor)
+		{
+			PrivateDependencyModuleNames.AddRange(
+			[
+				"UnrealEd",
+				"Engine"
+			]);
+		}
+	}
+}
diff --git a/Engine/Plugins/Performance/AutomatedPerfTestTools/Source/AutomatedPerfTestLaunchExtension/Private/AutomatedPerfTestLaunchExtension.cpp b/Engine/Plugins/Performance/AutomatedPerfTestTools/Source/AutomatedPerfTestLaunchExtension/Private/AutomatedPerfTestLaunchExtension.cpp
new file mode 100644
--- /dev/null
+++ b/Engine/Plugins/Performance/AutomatedPerfTestTools/Source/AutomatedPerfTestLaunchExtension/Private/AutomatedPerfTestLaunchExtension.cpp
@@ -0,0 +1,557 @@
+// Copyright Epic Games, Inc. All Rights Reserved.
+
+#include "AutomatedPerfTestLaunchExtension.h"
+#include "Templates/SharedPointer.h"
+#include "Model/ProjectLauncherModel.h"
+#include "Misc/Paths.h"
+#include "Misc/ConfigCacheIni.h"
+
+#include "Widgets/DeclarativeSyntaxSupport.h"
+#include "Widgets/Input/SButton.h"
+#include "Widgets/Shared/SCustomLaunchCombo.h"
+
+#define LOCTEXT_NAMESPACE "FAutomatedPerfTestLaunchExtension"
+
+namespace AutomatedPerfTestConfig
+{
+	using EConfigType = ProjectLauncher::FLaunchExtensionInstance::EConfig;
+	constexpr EConfigType ConfigTypeCommon = EConfigType::PerProfile;
+
+	constexpr const TCHAR* EnableCSVProfilerConfig = TEXT("EnableCSVProfiler");
+	constexpr const TCHAR* WindowedConfig = TEXT("Windowed");
+	constexpr const TCHAR* IterationsConfig = TEXT("Iterations");
+	constexpr const TCHAR* LLMConfig = TEXT("LLM");
+	constexpr const TCHAR* GPUPerfConfig = TEXT("GPUPerf");
+	constexpr const TCHAR* GPUReshapeConfig = TEXT("GPUReshape");
+	constexpr const TCHAR* LockDynamicResConfig = TEXT("LockDynamicRes");
+
+	constexpr const TCHAR* ProfileGoMapConfig = TEXT("ProfileGoMapName");
+	constexpr const TCHAR* ReplayConfig = TEXT("ReplayTestName");
+	constexpr const TCHAR* SequenceTestNameConfig = TEXT("SequenceTestName");
+	constexpr const TCHAR* StaticCameraMapConfig = TEXT("StaticCameraMapName");
+	constexpr const TCHAR* ProfileGoConfigFileName = TEXT("ProfileGoConfigFile");
+
+	constexpr EAutomatedPerfTestType DefaultTestType = EAutomatedPerfTestType::Sequence;
+	constexpr uint32 DefaultIterationCount = 1;
+
+	static constexpr bool IsEditorBuild()
+	{
+#if WITH_EDITOR
+		constexpr bool bIsEditor = true;
+#else
+		constexpr bool bIsEditor = false;
+#endif
+		return bIsEditor;
+	}
+}
+
+using namespace AutomatedPerfTestConfig;
+
+void FAutomatedPerfTestLaunchExtensionInstance::OnPropertyChanged()
+{
+	ILauncherProfileAutomatedTestPtr Test = GetTest();
+	if(Test.IsValid())
+	{
+		// This ensures the test type is always 
+		// updated to the current selected one. 
+		OnTestAdded(Test.ToSharedRef());
+	}
+
+	Super::OnPropertyChanged();
+}
+
+static constexpr const TCHAR* GetTestNodeName(EAutomatedPerfTestType TestType)
+{
+	// This is the string which represents test node defined in the
+	// Gauntlet controller. The test launch will fail if this string
+	// does not match the name of the controller class in Gauntlet.
+	switch (TestType)
+	{
+	case EAutomatedPerfTestType::Sequence:		return TEXT("AutomatedPerfTest.SequenceTest");
+	case EAutomatedPerfTestType::Replay:		return TEXT("AutomatedPerfTest.ReplayTest");
+	case EAutomatedPerfTestType::StaticCamera:	return TEXT("AutomatedPerfTest.StaticCameraTest");
+	case EAutomatedPerfTestType::Material:		return TEXT("AutomatedPerfTest.MaterialTest");
+	case EAutomatedPerfTestType::ProfileGo:		return TEXT("AutomatedPerfTest.ProfileGoTest");
+	default:									return TEXT("AutomatedPerfTest.DefaultTest");
+	}
+}
+
+void FAutomatedPerfTestLaunchExtensionInstance::OnTestAdded(ILauncherProfileAutomatedTestRef AutomatedTest)
+{
+	const TCHAR* TestNode = GetTestNodeName(GetTestType());
+
+	AutomatedTest->SetTests(TestNode);
+}
+
+const ProjectLauncher::FProjectSettings FAutomatedPerfTestLaunchExtensionInstance::GetProjectSettings() const
+{
+	using namespace ProjectLauncher;
+	TSharedRef<FModel> LauncherModel = GetModel();
+	return LauncherModel->GetProjectSettings(GetProfile());
+}
+
+const FConfigSection* FAutomatedPerfTestLaunchExtensionInstance::GetConfigSection(EAutomatedPerfTestType TestType) const
+{
+	const ProjectLauncher::FProjectSettings ProjectSettings = GetProjectSettings();
+	FConfigCacheIni* const Config = ProjectSettings.Config;
+	static const TMap<EAutomatedPerfTestType, FString> TestTypeSectionMap =
+	{
+		{ EAutomatedPerfTestType::Sequence,		TEXT("/Script/AutomatedPerfTesting.AutomatedSequencePerfTestProjectSettings") },
+		{ EAutomatedPerfTestType::Replay,		TEXT("/Script/AutomatedPerfTesting.AutomatedReplayPerfTestProjectSetting") },
+		{ EAutomatedPerfTestType::ProfileGo,	TEXT("/Script/AutomatedPerfTesting.AutomatedProfileGoTestProjectSetting") },
+		{ EAutomatedPerfTestType::StaticCamera,	TEXT("/Script/AutomatedPerfTesting.AutomatedStaticCameraPerfTestProjectSettings") },
+		{ EAutomatedPerfTestType::Material,		TEXT("/Script/AutomatedPerfTesting.AutomatedMaterialPerfTestProjectSetting") },
+	};
+
+	const FConfigSection* Section = Config->GetSection(*TestTypeSectionMap[TestType], true, GEngineIni);
+	return Section;
+}
+
+void FAutomatedPerfTestLaunchExtensionInstance::CustomizeTree(ProjectLauncher::FLaunchProfileTreeData& ProfileTreeData)
+{
+	using namespace ProjectLauncher;
+
+	const auto GetDisplayName = [](EAutomatedPerfTestType Type)
+	{
+		switch (Type)
+		{
+			case EAutomatedPerfTestType::Sequence:		return LOCTEXT("SequenceLabel", "Sequence");
+			case EAutomatedPerfTestType::Replay:		return LOCTEXT("ReplayLabel", "Replay");
+			case EAutomatedPerfTestType::ProfileGo:		return LOCTEXT("ProfileGoLabel", "ProfileGo");
+			case EAutomatedPerfTestType::StaticCamera:	return LOCTEXT("StaticCameraLabel", "StaticCamera");
+			case EAutomatedPerfTestType::Material:		return LOCTEXT("MaterialLabel", "Material");
+		}
+		return FText::GetEmpty();
+	};
+
+	const auto GetToolTip = [](EAutomatedPerfTestType Type)
+	{
+		switch (Type)
+		{
+		case EAutomatedPerfTestType::Sequence:		return LOCTEXT("SequenceToolTipLabel", "Automated Sequence Perf Test");
+		case EAutomatedPerfTestType::Replay:		return LOCTEXT("ReplayToolTipLabel", "Automated Replay Perf Test");
+		case EAutomatedPerfTestType::ProfileGo:		return LOCTEXT("ProfileGoToolTipLabel", "Automated ProfileGo Perf Test");
+		case EAutomatedPerfTestType::StaticCamera:	return LOCTEXT("StaticCameraToolTipLabel", "Automated StaticCamera Perf Test");
+		case EAutomatedPerfTestType::Material:		return LOCTEXT("MaterialToolTipLabel", "Automated Material Perf Test");
+		}
+		return FText::GetEmpty();
+	};
+
+	// TODO: Add support to specify sub test type i.e. LLM/Insights/GPU Perf etc. 
+	// Adding it as enum like the main test type above may not be ideal as sub-tests
+	// are expected to change more often. It may be prudent to consider that solution
+	// for the test types above as well to ensure test types can be added on and is 
+	// extensible in the future.
+
+	FLaunchProfileTreeNode& TreeNode = 
+		AddDefaultHeading(ProfileTreeData)
+		.AddBoolean(LOCTEXT("APTEnableCSVLabel", "Enable CSV Profiler"), 
+		{
+			.GetValue = [this]() { return GetConfigBool(ConfigTypeCommon, EnableCSVProfilerConfig, true); },
+			.SetValue = [this](bool bEnable) { return SetConfigBool(ConfigTypeCommon, EnableCSVProfilerConfig, bEnable); },
+		})
+		.AddBoolean(LOCTEXT("WindowedLabel","Windowed"),
+		{
+			.GetValue = [this]()          { return GetConfigBool(ConfigTypeCommon, WindowedConfig);},
+			.SetValue = [this](bool bVal) { SetConfigBool(ConfigTypeCommon, WindowedConfig, bVal ); },
+		})
+		.AddBoolean(LOCTEXT("LLMLabel","Enable LLM (Low Level Memory) Tracking"),
+		{
+			.GetValue = [this]()          { return GetConfigBool(ConfigTypeCommon, LLMConfig);},
+			.SetValue = [this](bool bVal) { SetConfigBool(ConfigTypeCommon, LLMConfig, bVal ); },
+		})
+		.AddBoolean(LOCTEXT("GPUPerfLabel","Enable GPU Perf"),
+		{
+			.GetValue = [this]()          { return GetConfigBool(ConfigTypeCommon, GPUPerfConfig);},
+			.SetValue = [this](bool bVal) { SetConfigBool(ConfigTypeCommon, GPUPerfConfig, bVal ); },
+		})
+		.AddBoolean(LOCTEXT("GPUReshapeLabel","Enable GPU Reshape"),
+		{
+			.GetValue = [this]()          { return GetConfigBool(ConfigTypeCommon, GPUReshapeConfig);},
+			.SetValue = [this](bool bVal) { SetConfigBool(ConfigTypeCommon, GPUReshapeConfig, bVal ); },
+		})
+		.AddBoolean(LOCTEXT("LockDynResLabel","Lock Dynamic Resolution (Default = 60%)"),
+		{
+			.GetValue = [this]()          { return GetConfigBool(ConfigTypeCommon, LockDynamicResConfig);},
+			.SetValue = [this](bool bVal) { SetConfigBool(ConfigTypeCommon, LockDynamicResConfig, bVal ); },
+		})
+		.AddInteger(LOCTEXT("IterationsLabel", "Iterations"), 
+		{
+			.GetValue = [this]() { return GetConfigInteger(ConfigTypeCommon, IterationsConfig, DefaultIterationCount); },
+			.SetValue = [this](int32 Value) { SetConfigInteger(ConfigTypeCommon, IterationsConfig, Value); },
+		})
+		.AddWidget(LOCTEXT("TestTypeLabel", "Test Type"),
+		{
+			.IsDefault = [this]() { return GetTestType() == DefaultTestType; },
+			.SetToDefault = [this]() { SetTestType(DefaultTestType); },
+		},
+		SNew(SCustomLaunchCombo<EAutomatedPerfTestType>)
+		.OnSelectionChanged(this, &FAutomatedPerfTestLaunchExtensionInstance::OnTestTypeSelectionChanged)
+		.SelectedItem(this, &FAutomatedPerfTestLaunchExtensionInstance::GetTestType)
+		.GetDisplayName_Lambda(GetDisplayName)
+		.GetItemToolTip_Lambda(GetToolTip)
+		.Items(TArray<EAutomatedPerfTestType>(
+			{
+				EAutomatedPerfTestType::Sequence,
+				EAutomatedPerfTestType::Replay ,
+				EAutomatedPerfTestType::ProfileGo ,
+				EAutomatedPerfTestType::StaticCamera ,
+				EAutomatedPerfTestType::Material
+			})
+		));
+
+	AddTestNodeOptions(TreeNode);
+}
+
+void FAutomatedPerfTestLaunchExtensionInstance::CustomizeUATCommandLine(FString& InOutCommandLine)
+{
+	// Most of the parameters here are defined and handled in the corresponding 
+	// Gauntlet controller which launches the automated perf test with the required 
+	// params based on the options passed down here.
+
+	if (GetConfigBool(ConfigTypeCommon, EnableCSVProfilerConfig, true))
+	{
+		InOutCommandLine += TEXT(" -AutomatedPerfTest.DoCSVProfiler");
+	}
+
+	if (GetConfigBool(ConfigTypeCommon, LLMConfig))
+	{
+		InOutCommandLine += TEXT(" -AutomatedPerfTest.DoLLM");
+	}
+
+	if (GetConfigBool(ConfigTypeCommon, GPUPerfConfig))
+	{
+		InOutCommandLine += TEXT(" -AutomatedPerfTest.DoGPUPerf");
+	}
+
+	if (GetConfigBool(ConfigTypeCommon, GPUReshapeConfig))
+	{
+		InOutCommandLine += TEXT(" -AutomatedPerfTest.DoGPUReshape");
+	}
+
+	if (GetConfigBool(ConfigTypeCommon, LockDynamicResConfig))
+	{
+		InOutCommandLine += TEXT(" -AutomatedPerfTest.LockDynamicRes");
+	}
+
+	const auto AddStringParamConfig = 
+	[this, &InOutCommandLine](const FString& Config, const FString& Param, const FString& DefaultVal = "")
+	{
+		FString ConfigVal = GetConfigString(ConfigTypeCommon, *Config, *DefaultVal);
+		if (!ConfigVal.IsEmpty())
+		{
+			InOutCommandLine += FString::Printf(TEXT(" -%s=\"%s\""), *Param, *EscapeJsonString(ConfigVal));
+		}
+	};
+
+	const int32 IterationCount = GetConfigInteger(ConfigTypeCommon, IterationsConfig, DefaultIterationCount);
+	InOutCommandLine += FString::Printf(TEXT(" -iterations=%d"), IterationCount);
+
+	const EAutomatedPerfTestType Type = GetTestType();
+	switch (Type)
+	{
+	case EAutomatedPerfTestType::ProfileGo: 
+	{
+	    FString ConfigPath = GetConfigString(ConfigTypeCommon, ProfileGoConfigFileName);
+		if(ConfigPath.IsEmpty())
+		{
+			ConfigPath = FPaths::Combine(FPaths::ProjectDir(), "Saved", "Profiling", "ProfileGo.json");
+		}
+
+		if (FPaths::FileExists(ConfigPath))
+		{
+			InOutCommandLine += FString::Printf(TEXT(" -profilego.config=\"%s\""), *EscapeJsonString(ConfigPath));
+		}
+
+		AddStringParamConfig(ProfileGoMapConfig, TEXT("Map"));
+		break;
+	}
+	case EAutomatedPerfTestType::Replay:
+	{
+		AddStringParamConfig(ReplayConfig, TEXT("AutomatedPerfTest.ReplayPerfTest.ReplayName"));
+		break;
+	}
+	case EAutomatedPerfTestType::Sequence:
+	{
+		AddStringParamConfig(SequenceTestNameConfig, TEXT("AutomatedPerfTest.SequencePerfTest.MapSequenceName"));
+		break;
+	}
+	case EAutomatedPerfTestType::Material:
+	{
+		// Material test uses settings for test parameters.
+		break;
+	}
+	case EAutomatedPerfTestType::StaticCamera:
+	{
+		AddStringParamConfig(StaticCameraMapConfig, TEXT("AutomatedPerfTest.StaticCameraPerfTest.MapName"));
+		break;
+	}
+	default: break;
+	}
+
+	if (GetConfigBool(ConfigTypeCommon, WindowedConfig))
+	{
+		InOutCommandLine += TEXT(" -windowed");
+	}
+}
+
+void FAutomatedPerfTestLaunchExtensionInstance::ExportProfileGoScenarios(const FString& Filename)
+{
+	// Note: This needs to exist outside this extension. Profile Go scenario creation and 
+	// export should ideally be handled in-editor. Disabling this for now and to be removed 
+	// once ProfileGo workflow is finalized. 
+#if 0
+	UWorld* World = GWorld;
+	UEditorEngine* EditorEngine = Cast<UEditorEngine>(GEngine);
+	if (GIsEditor && EditorEngine != nullptr && EditorEngine->PlayWorld != nullptr)
+	{
+		World = EditorEngine->PlayWorld.Get();
+	}
+
+	if (World)
+	{
+		// TODO: This needs to be standardized and/or made configurable with error handling. 
+		const FString OutputPath = FPaths::Combine(FPaths::ProjectDir(), "Saved", "Profiling", Filename);
+
+		UProfileGo& ProfileGo = UProfileGo::GetCDO();
+		ProfileGo.AddScenariosInLevel(World);
+		ProfileGo.SaveToJSON(OutputPath);
+	}
+#endif
+}
+
+void FAutomatedPerfTestLaunchExtensionInstance::AddTestNodeOptions(ProjectLauncher::FLaunchProfileTreeNode& TreeNode)
+{
+	AddSequenceTestNodeOptions(TreeNode);
+	AddReplayTestNodeOptions(TreeNode);
+	AddStaticCameraTestNodeOptions(TreeNode);
+	AddMaterialTestNodeOptions(TreeNode);
+	AddProfileGoTestNodeOptions(TreeNode);
+}
+
+void FAutomatedPerfTestLaunchExtensionInstance::AddSequenceTestNodeOptions(ProjectLauncher::FLaunchProfileTreeNode& TreeNode)
+{
+	constexpr auto TestType = EAutomatedPerfTestType::Sequence;
+	ProjectLauncher::FLaunchProfileTreeNode::FCallbacks&& Callbacks =
+		GetTestTypeCallbacks<TestType>();
+
+	TArray<FString> Sequences;
+	if(const FConfigSection* Section = GetConfigSection(TestType))
+	{
+		TArray<FString> ComboStructs;
+		Section->GetArray(TEXT("MapsAndSequencesToTest"), ComboStructs);
+
+		// The ini config has the combo name as well as other
+		// metadata which is not needed to launch the test. So
+		// we extract the combo name only. 
+		for (const FString& ComboStruct : ComboStructs)
+		{
+			constexpr TCHAR ComboStr[] = TEXT("ComboName=\"");
+			const int32 ComboPos = ComboStruct.Find(ComboStr);
+			const int32 FirstCommaPos = ComboStruct.Find(",");
+			if (ComboPos == INDEX_NONE || FirstCommaPos == INDEX_NONE)
+			{
+				continue;
+			}
+
+			const int32 StartIndex = ComboPos + FCString::Strlen(ComboStr);
+			const int32 EndIndex = FirstCommaPos - 1;
+			const int32 Length = EndIndex - StartIndex;
+			const FString ComboName = ComboStruct.Mid(StartIndex, Length);
+			Sequences.AddUnique(ComboName);
+		}
+	}
+
+	constexpr const TCHAR* ConfigName = SequenceTestNameConfig;
+	auto OnSequenceChanged = [this](FString Combo)
+	{
+		OnPropertyChanged();
+		SetConfigString(EConfig::PerProfile, ConfigName, Combo);
+	};
+
+	auto GetSequence = [this]() { return GetConfigString(EConfig::PerProfile, ConfigName); };
+	auto GetDisplayName = [this](FString Combo) -> FText { return FText::FromString(GetConfigString(EConfig::PerProfile, ConfigName, *Combo)); };
+
+	TreeNode.AddWidget(LOCTEXT("SequenceComboNameLabel", "Sequence Combo Name"),
+		MoveTemp(Callbacks),
+		SNew(SCustomLaunchCombo<FString>)
+		.OnSelectionChanged_Lambda(MoveTemp(OnSequenceChanged))
+		.SelectedItem_Lambda(MoveTemp(GetSequence))
+		.GetDisplayName_Lambda(MoveTemp(GetDisplayName))
+		.Items(Sequences)
+	);
+}
+
+void FAutomatedPerfTestLaunchExtensionInstance::AddReplayTestNodeOptions(ProjectLauncher::FLaunchProfileTreeNode& TreeNode)
+{
+	ProjectLauncher::FLaunchProfileTreeNode::FCallbacks&& Callbacks =
+		GetTestTypeCallbacks<EAutomatedPerfTestType::Replay>();
+
+	const TCHAR* FileTypeFilter = TEXT("Replay files (*.replay)|*.replay");
+	constexpr const TCHAR* ConfigName = ReplayConfig;
+	TreeNode.AddFileString(LOCTEXT("ReplayTestNameLabel", "Replay Name"), 
+	{
+		.GetValue = [this]() { return GetConfigString(EConfig::PerProfile, ConfigName); },
+		.SetValue = [this](FString Value) { SetConfigString(EConfig::PerProfile, ConfigName, Value); },
+		.IsVisible = Callbacks.IsVisible,
+		.IsEnabled = Callbacks.IsEnabled
+	}, FileTypeFilter);
+}
+
+void FAutomatedPerfTestLaunchExtensionInstance::AddProfileGoTestNodeOptions(ProjectLauncher::FLaunchProfileTreeNode& TreeNode)
+{
+    ProjectLauncher::FLaunchProfileTreeNode::FCallbacks&& Callbacks = 
+		GetTestTypeCallbacks<EAutomatedPerfTestType::ProfileGo>();
+
+	constexpr const TCHAR* ConfigName = ProfileGoMapConfig;
+	constexpr const TCHAR* ConfigFileName = ProfileGoConfigFileName;
+	constexpr const TCHAR* FileTypeFilter = TEXT("ProfileGo Config JSON (*.json)|*.json");
+
+	constexpr bool bIncludeNonContentDirMaps = true;
+	const FString ProjectPath = GetProfile()->GetProjectBasePath();
+	TArray<FString> Maps = GetModel()->GetAvailableProjectMapNames(ProjectPath, bIncludeNonContentDirMaps);
+
+	auto OnMapChanged = [this](FString Map)
+	{
+		OnPropertyChanged();
+		SetConfigString(ConfigTypeCommon, ConfigName, Map);
+	};
+
+	auto GetMap = [this]() { return GetConfigString(ConfigTypeCommon, ConfigName); };
+	auto GetDisplayName = [this](FString Map) -> FText { return FText::FromString(*Map); };
+
+	TreeNode.AddWidget(LOCTEXT("TestExportLabel", "Export ProfileGo"),
+		MoveTemp(Callbacks),
+		SNew(SButton)
+		.HAlign(HAlign_Center)
+		.VAlign(VAlign_Center)
+		.Text(LOCTEXT("TestExportLabelButton", "Export JSON"))
+		.OnClicked_Lambda([this]() -> FReply
+		{
+			// TEMP: Will be removed once ProfileGo scenario workflow creation is finalized.
+			const TCHAR* ProfileGoFilename = TEXT("ProfileGo.json");
+			ExportProfileGoScenarios(ProfileGoFilename);
+			return FReply::Handled();
+		})
+		.IsEnabled_Lambda([this]() { return GetTestType() == EAutomatedPerfTestType::ProfileGo && IsEditorBuild(); })
+	)
+	.AddWidget(LOCTEXT("SettingsUpdateLabel", "Update ProfileGo Settings"),
+		MoveTemp(Callbacks),
+		SNew(SButton)
+		.HAlign(HAlign_Center)
+		.VAlign(VAlign_Center)
+		.Text(LOCTEXT("SettingsUpdateLabelButton", "Update"))
+		.OnClicked_Lambda([this]() -> FReply
+		{
+		// This needs to exist outside this extension. Disabling for now. 
+		#if 0
+			// TEMP: Will be removed once ProfileGo scenario workflow creation is finalized.
+			UWorld* World = GWorld;
+			UEditorEngine* EditorEngine = Cast<UEditorEngine>(GEngine);
+			if (GIsEditor && EditorEngine != nullptr && EditorEngine->PlayWorld != nullptr)
+			{
+				World = EditorEngine->PlayWorld.Get();
+			}
+			UProfileGo::GetCDO().AddScenariosInLevel(World);
+		#endif
+			return FReply::Handled();
+		})
+		.IsEnabled_Lambda([this]() { return GetTestType() == EAutomatedPerfTestType::ProfileGo && IsEditorBuild(); })
+	)
+	.AddFileString(LOCTEXT("ProfileGoConfigLabel","Config File"),
+	{
+		.GetValue = [this]() { return GetConfigString(EConfig::PerProfile, ConfigFileName); },
+		.SetValue = [this](FString Value) { SetConfigString(EConfig::PerProfile, ConfigFileName, Value); },
+		.IsVisible = Callbacks.IsVisible,
+		.IsEnabled = Callbacks.IsEnabled
+	}, 
+	FileTypeFilter)
+	.AddWidget(LOCTEXT("ProfileGoMapLabel", "Map"),
+		MoveTemp(Callbacks),
+		SNew(SCustomLaunchCombo<FString>)
+		.OnSelectionChanged_Lambda(MoveTemp(OnMapChanged))
+		.SelectedItem_Lambda(MoveTemp(GetMap))
+		.GetDisplayName_Lambda(MoveTemp(GetDisplayName))
+		.Items(MoveTemp(Maps))
+	);
+}
+
+void FAutomatedPerfTestLaunchExtensionInstance::AddStaticCameraTestNodeOptions(ProjectLauncher::FLaunchProfileTreeNode& TreeNode)
+{
+	constexpr auto TestType = EAutomatedPerfTestType::StaticCamera;
+	ProjectLauncher::FLaunchProfileTreeNode::FCallbacks&& Callbacks =
+		GetTestTypeCallbacks<TestType>();
+
+	constexpr const TCHAR* ConfigName = StaticCameraMapConfig;
+
+	TArray<FString> Maps;
+	if(const FConfigSection* Section = GetConfigSection(TestType))
+	{
+		Section->GetArray(TEXT("MapsToTest"), Maps);
+	}
+
+	auto OnMapChanged = [this](FString Map)
+	{
+		OnPropertyChanged();
+		SetConfigString(EConfig::PerProfile, ConfigName, Map);
+	};
+
+	auto GetMap = [this]() { return GetConfigString(EConfig::PerProfile, ConfigName); };
+	auto GetDisplayName = [this](FString Map) -> FText { return FText::FromString(*Map); };
+
+	TreeNode.AddWidget(LOCTEXT("StaticCameraTestNameLabel", "Static Camera Map Name"), 
+		 MoveTemp(Callbacks),
+		 SNew(SCustomLaunchCombo<FString>)
+		.OnSelectionChanged_Lambda(MoveTemp(OnMapChanged))
+		.SelectedItem_Lambda(MoveTemp(GetMap))
+		.GetDisplayName_Lambda(MoveTemp(GetDisplayName))
+		.Items(MoveTemp(Maps))
+	);
+}
+
+void FAutomatedPerfTestLaunchExtensionInstance::AddMaterialTestNodeOptions(ProjectLauncher::FLaunchProfileTreeNode& TreeNode)
+{
+	// Material test does not have specific params.
+}
+
+void FAutomatedPerfTestLaunchExtensionInstance::OnTestTypeSelectionChanged(EAutomatedPerfTestType Type)
+{
+	SetTestType(Type);
+	OnPropertyChanged();
+}
+
+void FAutomatedPerfTestLaunchExtensionInstance::SetTestType(EAutomatedPerfTestType Type)
+{
+	CurrentTestType = Type;
+	SetConfigInteger(EConfig::PerProfile, TEXT("TestType"), (int32)Type);
+	BroadcastPropertyChanged();
+}
+
+EAutomatedPerfTestType FAutomatedPerfTestLaunchExtensionInstance::GetTestType() const
+{
+	return (EAutomatedPerfTestType)GetConfigInteger(EConfig::PerProfile, TEXT("TestType"), (int32)CurrentTestType);
+}
+
+TSharedPtr<ProjectLauncher::FLaunchExtensionInstance> FAutomatedPerfTestLaunchExtension::CreateInstanceForProfile(ProjectLauncher::FLaunchExtensionInstance::FArgs& InArgs)
+{
+	return MakeShared<FAutomatedPerfTestLaunchExtensionInstance>(InArgs);
+}
+
+const TCHAR* FAutomatedPerfTestLaunchExtension::GetInternalName() const
+{
+	return TEXT("AutomatedPerfTest");
+}
+
+FText FAutomatedPerfTestLaunchExtension::GetDisplayName() const
+{
+	return LOCTEXT("ExtensionName", "Automated Perf Test");
+}
+
+void FAutomatedPerfTestLaunchExtension::GetExtensionsMenuEntry(FExtensionsMenuEntry& MenuEntry) const
+{
+	MenuEntry = FExtensionsMenuEntry::OwnSection;
+}
+
+#undef LOCTEXT_NAMESPACE
+
+
diff --git a/Engine/Plugins/Performance/AutomatedPerfTestTools/Source/AutomatedPerfTestLaunchExtension/Private/AutomatedPerfTestLaunchExtension.h b/Engine/Plugins/Performance/AutomatedPerfTestTools/Source/AutomatedPerfTestLaunchExtension/Private/AutomatedPerfTestLaunchExtension.h
new file mode 100644
--- /dev/null
+++ b/Engine/Plugins/Performance/AutomatedPerfTestTools/Source/AutomatedPerfTestLaunchExtension/Private/AutomatedPerfTestLaunchExtension.h
@@ -0,0 +1,78 @@
+// Copyright Epic Games, Inc. All Rights Reserved.
+
+#pragma once
+
+#include "Extension/AutomatedTestLaunchExtension.h"
+#include "Model/ProjectLauncherModel.h"
+#include "Templates/UnrealTemplate.h"
+
+enum class EAutomatedPerfTestType : uint8
+{
+	Sequence,
+	Replay,
+	ProfileGo,
+	StaticCamera,
+	Material,
+	MAX
+};
+
+class FAutomatedPerfTestLaunchExtensionInstance : public ProjectLauncher::FAutomatedTestLaunchExtensionInstance
+{
+	using Super = ProjectLauncher::FAutomatedTestLaunchExtensionInstance;
+public:
+	FAutomatedPerfTestLaunchExtensionInstance(FArgs& InArgs) : FAutomatedTestLaunchExtensionInstance(InArgs) {};
+	virtual ~FAutomatedPerfTestLaunchExtensionInstance() = default;
+
+	virtual void OnPropertyChanged() override;
+	virtual void OnTestAdded(ILauncherProfileAutomatedTestRef AutomatedTest) override;
+
+	virtual void CustomizeTree(ProjectLauncher::FLaunchProfileTreeData& ProfileTreeData) override;
+	virtual void CustomizeUATCommandLine(FString& InOutCommandLine) override;
+
+private:
+	void ExportProfileGoScenarios(const FString& Filename);
+
+	void AddTestNodeOptions(ProjectLauncher::FLaunchProfileTreeNode& TreeNode);
+	void AddSequenceTestNodeOptions(ProjectLauncher::FLaunchProfileTreeNode& TreeNode);
+	void AddReplayTestNodeOptions(ProjectLauncher::FLaunchProfileTreeNode& TreeNode);
+	void AddProfileGoTestNodeOptions(ProjectLauncher::FLaunchProfileTreeNode& TreeNode);
+	void AddStaticCameraTestNodeOptions(ProjectLauncher::FLaunchProfileTreeNode& TreeNode);
+	void AddMaterialTestNodeOptions(ProjectLauncher::FLaunchProfileTreeNode& TreeNode);
+
+	void OnTestTypeSelectionChanged(EAutomatedPerfTestType Type);
+	void SetTestType(EAutomatedPerfTestType Type);
+	EAutomatedPerfTestType GetTestType() const;
+
+	const ProjectLauncher::FProjectSettings GetProjectSettings() const;
+	const FConfigSection* GetConfigSection(EAutomatedPerfTestType TestType) const;
+
+	template<EAutomatedPerfTestType TestType>
+	ProjectLauncher::FLaunchProfileTreeNode::FCallbacks GetTestTypeCallbacks()
+	{
+		const auto IsTestType = [this]()
+		{
+			return GetTestType() == TestType;
+		};
+
+		ProjectLauncher::FLaunchProfileTreeNode::FCallbacks Callbacks
+		{
+			.IsVisible = IsTestType,
+			.IsEnabled = IsTestType
+		};
+
+		return MoveTemp(Callbacks);
+	}
+
+	EAutomatedPerfTestType CurrentTestType = EAutomatedPerfTestType::Sequence;
+};
+
+
+class FAutomatedPerfTestLaunchExtension : public ProjectLauncher::FAutomatedTestLaunchExtension
+{
+public:
+	virtual TSharedPtr<ProjectLauncher::FLaunchExtensionInstance> CreateInstanceForProfile(ProjectLauncher::FLaunchExtensionInstance::FArgs& InArgs) override;
+	virtual const TCHAR* GetInternalName() const override;
+	virtual FText GetDisplayName() const override;
+	virtual void GetExtensionsMenuEntry(FExtensionsMenuEntry& MenuEntry) const override;
+};
+
diff --git a/Engine/Plugins/Performance/AutomatedPerfTestTools/Source/AutomatedPerfTestLaunchExtension/Private/AutomatedPerfTestLaunchExtensionModule.cpp b/Engine/Plugins/Performance/AutomatedPerfTestTools/Source/AutomatedPerfTestLaunchExtension/Private/AutomatedPerfTestLaunchExtensionModule.cpp
new file mode 100644
--- /dev/null
+++ b/Engine/Plugins/Performance/AutomatedPerfTestTools/Source/AutomatedPerfTestLaunchExtension/Private/AutomatedPerfTestLaunchExtensionModule.cpp
@@ -0,0 +1,28 @@
+// Copyright Epic Games, Inc. All Rights Reserved.
+
+#include "AutomatedPerfTestLaunchExtensionModule.h"
+
+void FAutomatedPerfTestLaunchExtensionModule::StartupModule() 
+{
+	Extensions.Add(MakeShared<FAutomatedPerfTestLaunchExtension>());
+	for (TSharedRef<ProjectLauncher::FLaunchExtension> Extension : Extensions)
+	{
+		IProjectLauncherModule::Get().RegisterExtension(Extension);
+	}
+}
+
+void FAutomatedPerfTestLaunchExtensionModule::ShutdownModule() 
+{
+	if (IProjectLauncherModule* ProjectLauncher = IProjectLauncherModule::TryGet())
+	{
+		for (TSharedRef<ProjectLauncher::FLaunchExtension> Extension : Extensions)
+		{
+			ProjectLauncher->UnregisterExtension(Extension);
+		}
+	}
+
+	Extensions.Reset();
+}
+
+IMPLEMENT_MODULE(FAutomatedPerfTestLaunchExtensionModule, AutomatedPerfTestLaunchExtension);
+
diff --git a/Engine/Plugins/Performance/AutomatedPerfTestTools/Source/AutomatedPerfTestLaunchExtension/Public/AutomatedPerfTestLaunchExtensionModule.h b/Engine/Plugins/Performance/AutomatedPerfTestTools/Source/AutomatedPerfTestLaunchExtension/Public/AutomatedPerfTestLaunchExtensionModule.h
new file mode 100644
--- /dev/null
+++ b/Engine/Plugins/Performance/AutomatedPerfTestTools/Source/AutomatedPerfTestLaunchExtension/Public/AutomatedPerfTestLaunchExtensionModule.h
@@ -0,0 +1,20 @@
+// Copyright Epic Games, Inc. All Rights Reserved.
+
+#pragma once
+
+#include "Modules/ModuleInterface.h"
+#include "ProjectLauncherModule.h"
+#include "AutomatedPerfTestLaunchExtension.h"
+
+class FAutomatedPerfTestLaunchExtensionModule : public IModuleInterface
+{
+public:
+	virtual void StartupModule() override;
+	virtual void ShutdownModule() override;
+
+private:
+	TArray<TSharedRef<ProjectLauncher::FLaunchExtension>> Extensions;
+};
+
+
+
diff --git a/Engine/Plugins/Performance/AutomatedPerfTesting/Source/AutomatedPerfTesting/Private/AutomatedPerfTesting.cpp b/Engine/Plugins/Performance/AutomatedPerfTesting/Source/AutomatedPerfTesting/Private/AutomatedPerfTesting.cpp
--- a/Engine/Plugins/Performance/AutomatedPerfTesting/Source/AutomatedPerfTesting/Private/AutomatedPerfTesting.cpp
+++ b/Engine/Plugins/Performance/AutomatedPerfTesting/Source/AutomatedPerfTesting/Private/AutomatedPerfTesting.cpp
@@ -2,36 +2,15 @@
 
 #include "AutomatedPerfTesting.h"
 
-#include "Modules/ModuleInterface.h"
-#include "Templates/SharedPointer.h"
-#include "LaunchExtension/AutomatedPerfTestLaunchExtension.h"
-
-#if WITH_EDITOR
-#include "ProjectLauncherModule.h"
-#endif
-
 #define LOCTEXT_NAMESPACE "FAutomatedPerfTestingModule"
 
 void FAutomatedPerfTestingModule::StartupModule()
 {
-#if WITH_EDITOR
-	LaunchExtension = MakeShared<FAutomatedPerfTestLaunchExtension>();
-	IProjectLauncherModule::Get().RegisterExtension(LaunchExtension.ToSharedRef());
-#endif
 }
 
 void FAutomatedPerfTestingModule::ShutdownModule()
 {
-#if WITH_EDITOR
-	if (IProjectLauncherModule* ProjectLauncher = IProjectLauncherModule::TryGet())
-	{
-		ProjectLauncher->UnregisterExtension(LaunchExtension.ToSharedRef());
-	}
-
-	LaunchExtension.Reset();
-#endif
 }
 
 #undef LOCTEXT_NAMESPACE
-	
 IMPLEMENT_MODULE(FAutomatedPerfTestingModule, AutomatedPerfTesting)
\ No newline at end of file
diff --git a/Engine/Plugins/Performance/AutomatedPerfTesting/Source/AutomatedPerfTesting/Private/LaunchExtension/AutomatedPerfTestLaunchExtension.cpp b/Engine/Plugins/Performance/AutomatedPerfTesting/Source/AutomatedPerfTesting/Private/LaunchExtension/AutomatedPerfTestLaunchExtension.cpp
deleted file mode 100644
diff --git a/Engine/Plugins/Performance/AutomatedPerfTesting/Source/AutomatedPerfTesting/Private/LaunchExtension/AutomatedPerfTestLaunchExtension.h b/Engine/Plugins/Performance/AutomatedPerfTesting/Source/AutomatedPerfTesting/Private/LaunchExtension/AutomatedPerfTestLaunchExtension.h
deleted file mode 100644
diff --git a/Engine/Plugins/Performance/AutomatedPerfTesting/Source/AutomatedPerfTesting/Public/AutomatedPerfTesting.h b/Engine/Plugins/Performance/AutomatedPerfTesting/Source/AutomatedPerfTesting/Public/AutomatedPerfTesting.h
--- a/Engine/Plugins/Performance/AutomatedPerfTesting/Source/AutomatedPerfTesting/Public/AutomatedPerfTesting.h
+++ b/Engine/Plugins/Performance/AutomatedPerfTesting/Source/AutomatedPerfTesting/Public/AutomatedPerfTesting.h
@@ -9,8 +9,6 @@
 DECLARE_LOG_CATEGORY_EXTERN(LogAutomatedPerfTest, Log, All)
 CSV_DECLARE_CATEGORY_EXTERN(AutomatedPerfTest);
 
-class FAutomatedPerfTestLaunchExtension;
-
 class FAutomatedPerfTestingModule : public IModuleInterface
 {
 public:
@@ -20,7 +18,4 @@ public:
 	virtual void ShutdownModule() override;
 
 private:
-#if WITH_EDITOR
-	TSharedPtr<FAutomatedPerfTestLaunchExtension> LaunchExtension;
-#endif
 };
