APT Launcher Extension Minor improvements and updates

* Add optional Hint Text and Tool Tip text parameters for String input and boolean input where applicablke.
* Move options to "Advanced" subheading in APT Task
* Add options to specify `-deterministic` and to pass cvars to test instance in Gauntlet C# node.
* Add ability to specify GPU Perf, TestID, insights trace and channels to APT Launcher extension

#tests Local APT Sequence Tests with different options enabled.
#rnx

#Synced-CL 51902396

diff --git a/Engine/Plugins/Developer/ProjectLauncher/Source/ProjectLauncher/Private/ProfileTree/LaunchProfileTreeData.cpp b/Engine/Plugins/Developer/ProjectLauncher/Source/ProjectLauncher/Private/ProfileTree/LaunchProfileTreeData.cpp
--- a/Engine/Plugins/Developer/ProjectLauncher/Source/ProjectLauncher/Private/ProfileTree/LaunchProfileTreeData.cpp
+++ b/Engine/Plugins/Developer/ProjectLauncher/Source/ProjectLauncher/Private/ProfileTree/LaunchProfileTreeData.cpp
@@ -390,7 +390,7 @@ namespace ProjectLauncher
 	}
 
 
-	FLaunchProfileTreeNode& FLaunchProfileTreeNode::AddBoolean( FText InName, FBooleanCallbacks&& BooleanCallbacks )
+	FLaunchProfileTreeNode& FLaunchProfileTreeNode::AddBoolean( FText InName, FBooleanCallbacks&& BooleanCallbacks, FText ToolTipText)
 	{
 		check(BooleanCallbacks.GetValue);
 		check(BooleanCallbacks.SetValue);
@@ -430,16 +430,22 @@ namespace ProjectLauncher
 			return !BooleanCallbacks.IsEnabled.IsSet() || BooleanCallbacks.IsEnabled();
 		};
 
+		auto GetToolTipText = [ToolTipText]()
+		{
+			return ToolTipText;
+		};
+
 		return AddWidget( InName,
 			MoveTemp(WidgetCallbacks),
 			SNew(SCheckBox)
 			.OnCheckStateChanged_Lambda(SetCheckState)
 			.IsChecked_Lambda(GetCheckState)
 			.IsEnabled_Lambda(IsEnabled)
+			.ToolTipText_Lambda(GetToolTipText)
 		);
 	}
 
-	FLaunchProfileTreeNode& FLaunchProfileTreeNode::AddString( FText InName, FStringCallbacks&& StringCallbacks )
+	FLaunchProfileTreeNode& FLaunchProfileTreeNode::AddString( FText InName, FStringCallbacks&& StringCallbacks, FText ToolTipText, FText HintText)
 	{
 		check(StringCallbacks.GetValue);
 		check(StringCallbacks.SetValue);
@@ -479,6 +485,15 @@ namespace ProjectLauncher
 			return StringCallbacks.IsEnabled && !StringCallbacks.IsEnabled();
 		};
 
+		auto GetHintText = [HintText]()
+		{
+			return HintText;
+		};
+
+		auto GetToolTipText = [ToolTipText]()
+		{
+			return ToolTipText;
+		};
 
 		return AddWidget( InName,
 			MoveTemp(WidgetCallbacks),
@@ -489,6 +504,8 @@ namespace ProjectLauncher
 			.Text_Lambda( GetString )
 			.OnTextCommitted_Lambda( SetString )
 			.IsReadOnly_Lambda( IsReadOnly )
+			.HintText_Lambda(GetHintText)
+			.ToolTipText_Lambda(GetToolTipText)
 		);
 
 	}
diff --git a/Engine/Plugins/Developer/ProjectLauncher/Source/ProjectLauncher/Public/ProfileTree/LaunchProfileTreeData.h b/Engine/Plugins/Developer/ProjectLauncher/Source/ProjectLauncher/Public/ProfileTree/LaunchProfileTreeData.h
--- a/Engine/Plugins/Developer/ProjectLauncher/Source/ProjectLauncher/Public/ProfileTree/LaunchProfileTreeData.h
+++ b/Engine/Plugins/Developer/ProjectLauncher/Source/ProjectLauncher/Public/ProfileTree/LaunchProfileTreeData.h
@@ -109,7 +109,7 @@ namespace ProjectLauncher
 			FGetBool IsEnabled = nullptr;
 			FValidation Validation;
 		};
-		UE_API FLaunchProfileTreeNode& AddBoolean( FText InName, FBooleanCallbacks&& BooleanCallbacks );
+		UE_API FLaunchProfileTreeNode& AddBoolean( FText InName, FBooleanCallbacks&& BooleanCallbacks, FText ToolTipText = FText::FromString(""));
 
 		/* String-specific callbacks. */
 		struct FStringCallbacks
@@ -121,7 +121,7 @@ namespace ProjectLauncher
 			FGetBool IsEnabled = nullptr;
 			FValidation Validation;
 		};
-		UE_API FLaunchProfileTreeNode& AddString( FText InName, FStringCallbacks&& StringCallbacks );
+		UE_API FLaunchProfileTreeNode& AddString( FText InName, FStringCallbacks&& StringCallbacks, FText ToolTipText = FText::FromString(""), FText HintText = FText::FromString(""));
 		UE_API FLaunchProfileTreeNode& AddDirectoryString(FText InName, FStringCallbacks&& StringCallbacks);
 		UE_API FLaunchProfileTreeNode& AddFileString(FText InName, FStringCallbacks&& StringCallbacks, FString FileFilter = TEXT("All files (*.*)|*.*"));
 		UE_API FLaunchProfileTreeNode& AddCommandLineString( FText InName, FStringCallbacks&& StringCallbacks, const TCHAR* InType = nullptr ); // defaults to "Launch" cmdline extensions
diff --git a/Engine/Plugins/Performance/AutomatedPerfTestTools/Source/AutomatedPerfTestLaunchExtension/Private/AutomatedPerfTestLaunchExtension.cpp b/Engine/Plugins/Performance/AutomatedPerfTestTools/Source/AutomatedPerfTestLaunchExtension/Private/AutomatedPerfTestLaunchExtension.cpp
--- a/Engine/Plugins/Performance/AutomatedPerfTestTools/Source/AutomatedPerfTestLaunchExtension/Private/AutomatedPerfTestLaunchExtension.cpp
+++ b/Engine/Plugins/Performance/AutomatedPerfTestTools/Source/AutomatedPerfTestLaunchExtension/Private/AutomatedPerfTestLaunchExtension.cpp
@@ -24,6 +24,11 @@ namespace AutomatedPerfTestConfig
 	constexpr const TCHAR* GPUPerfConfig = TEXT("GPUPerf");
 	constexpr const TCHAR* GPUReshapeConfig = TEXT("GPUReshape");
 	constexpr const TCHAR* LockDynamicResConfig = TEXT("LockDynamicRes");
+	constexpr const TCHAR* CVarConfig = TEXT("CVars");
+	constexpr const TCHAR* TestIDConfig = TEXT("TestID");
+	constexpr const TCHAR* DeterministicConfig = TEXT("Deterministic");
+	constexpr const TCHAR* InsightsConfig = TEXT("Insights");
+	constexpr const TCHAR* InsightsTraceChannelsConfig = TEXT("InsightsTraceChannels");
 
 	constexpr const TCHAR* ProfileGoMapConfig = TEXT("ProfileGoMapName");
 	constexpr const TCHAR* ReplayConfig = TEXT("ReplayTestName");
@@ -137,44 +142,19 @@ void FAutomatedPerfTestLaunchExtensionInstance::CustomizeTree(ProjectLauncher::F
 		return FText::GetEmpty();
 	};
 
-	// TODO: Add support to specify sub test type i.e. LLM/Insights/GPU Perf etc. 
-	// Adding it as enum like the main test type above may not be ideal as sub-tests
-	// are expected to change more often. It may be prudent to consider that solution
-	// for the test types above as well to ensure test types can be added on and is 
-	// extensible in the future.
-
 	FLaunchProfileTreeNode& TreeNode = 
 		AddDefaultHeading(ProfileTreeData)
 		.AddBoolean(LOCTEXT("APTEnableCSVLabel", "Enable CSV Profiler"), 
 		{
 			.GetValue = [this]() { return GetConfigBool(ConfigTypeCommon, EnableCSVProfilerConfig, true); },
 			.SetValue = [this](bool bEnable) { return SetConfigBool(ConfigTypeCommon, EnableCSVProfilerConfig, bEnable); },
+			.GetDefaultValue = [this]() { return true; }
 		})
 		.AddBoolean(LOCTEXT("WindowedLabel","Windowed"),
 		{
 			.GetValue = [this]()          { return GetConfigBool(ConfigTypeCommon, WindowedConfig);},
 			.SetValue = [this](bool bVal) { SetConfigBool(ConfigTypeCommon, WindowedConfig, bVal ); },
 		})
-		.AddBoolean(LOCTEXT("LLMLabel","Enable LLM (Low Level Memory) Tracking"),
-		{
-			.GetValue = [this]()          { return GetConfigBool(ConfigTypeCommon, LLMConfig);},
-			.SetValue = [this](bool bVal) { SetConfigBool(ConfigTypeCommon, LLMConfig, bVal ); },
-		})
-		.AddBoolean(LOCTEXT("GPUPerfLabel","Enable GPU Perf"),
-		{
-			.GetValue = [this]()          { return GetConfigBool(ConfigTypeCommon, GPUPerfConfig);},
-			.SetValue = [this](bool bVal) { SetConfigBool(ConfigTypeCommon, GPUPerfConfig, bVal ); },
-		})
-		.AddBoolean(LOCTEXT("GPUReshapeLabel","Enable GPU Reshape"),
-		{
-			.GetValue = [this]()          { return GetConfigBool(ConfigTypeCommon, GPUReshapeConfig);},
-			.SetValue = [this](bool bVal) { SetConfigBool(ConfigTypeCommon, GPUReshapeConfig, bVal ); },
-		})
-		.AddBoolean(LOCTEXT("LockDynResLabel","Lock Dynamic Resolution (Default = 60%)"),
-		{
-			.GetValue = [this]()          { return GetConfigBool(ConfigTypeCommon, LockDynamicResConfig);},
-			.SetValue = [this](bool bVal) { SetConfigBool(ConfigTypeCommon, LockDynamicResConfig, bVal ); },
-		})
 		.AddInteger(LOCTEXT("IterationsLabel", "Iterations"), 
 		{
 			.GetValue = [this]() { return GetConfigInteger(ConfigTypeCommon, IterationsConfig, DefaultIterationCount); },
@@ -200,7 +180,71 @@ void FAutomatedPerfTestLaunchExtensionInstance::CustomizeTree(ProjectLauncher::F
 			})
 		));
 
-	AddTestNodeOptions(TreeNode);
+	// Test type specific options will be added here. 
+	FLaunchProfileTreeNode& TestTypeTreeNode = TreeNode.AddSubHeading(TEXT("TestTypeOptions"), LOCTEXT("TestTypeOptionsLabel", "Test Type Options"));
+	AddTestNodeOptions(TestTypeTreeNode);
+
+	// Add "Advanced" options 
+	TreeNode
+		.AddSubHeading(TEXT("Advanced"), LOCTEXT("APTAdvancedLabel", "Advanced"))
+		.AddString(LOCTEXT("TestIDLabel", "Test ID"),
+		{
+			.GetValue = [this]() { return GetConfigString(ConfigTypeCommon, TestIDConfig); },
+			.SetValue = [this](FString TestID) { return SetConfigString(ConfigTypeCommon, TestIDConfig, TestID); }
+		}, 
+		LOCTEXT("TestIDHint", "My-Test-Identifier"), 
+		LOCTEXT("TestIDTooltip", "If you want to differentiate between runs, give this test run a name or ID"))
+		.AddBoolean(LOCTEXT("LLMLabel","Enable LLM"),
+		{
+			.GetValue = [this]()          { return GetConfigBool(ConfigTypeCommon, LLMConfig);},
+			.SetValue = [this](bool bVal) { SetConfigBool(ConfigTypeCommon, LLMConfig, bVal ); },
+		},
+		LOCTEXT("LLMToolTip", "Enables Low Level Memory (LLM) Tracking Stats"))
+		.AddBoolean(LOCTEXT("InsightsLabel","Enable Insights"),
+		{
+			.GetValue = [this]()          { return GetConfigBool(ConfigTypeCommon, InsightsConfig);},
+			.SetValue = [this](bool bVal) { SetConfigBool(ConfigTypeCommon, InsightsConfig, bVal ); },
+		},
+		LOCTEXT("InsightsToolTip", "Enables Insights Trace"))
+		.AddBoolean(LOCTEXT("DeterministicLabel","Deterministic"),
+		{
+			.GetValue = [this]()          { return GetConfigBool(ConfigTypeCommon, DeterministicConfig);},
+			.SetValue = [this](bool bVal) { SetConfigBool(ConfigTypeCommon, DeterministicConfig, bVal ); },
+			.GetDefaultValue = [this]()	  { return true; }
+		})
+		.AddString(LOCTEXT("TraceChannelsLabel", "Trace Channels"),
+		{
+			.GetValue = [this]() { return GetConfigString(ConfigTypeCommon, InsightsTraceChannelsConfig); },
+			.SetValue = [this](FString Channels) { return SetConfigString(ConfigTypeCommon, InsightsTraceChannelsConfig, Channels); },
+			.GetDefaultValue = [this]() { return TEXT("default,screenshot,stats"); }
+		}, 
+		LOCTEXT("TraceChannelsHint", "default,channel1,channel2,..."),
+		LOCTEXT("TraceChannelsTooltip", "Trace Channels to use if Insights Trace is enabled. Comma separated values. Default=default,screenshot,stats"))
+		.AddBoolean(LOCTEXT("GPUPerfLabel","Enable GPU Perf"),
+		{
+			.GetValue = [this]()          { return GetConfigBool(ConfigTypeCommon, GPUPerfConfig);},
+			.SetValue = [this](bool bVal) { SetConfigBool(ConfigTypeCommon, GPUPerfConfig, bVal ); },
+		},
+
+		LOCTEXT("GPUPerfHint", "Enables GPU Perf related CVars and locks dynamic resolution"))
+		.AddBoolean(LOCTEXT("GPUReshapeLabel","Enable GPU Reshape"),
+		{
+			.GetValue = [this]()          { return GetConfigBool(ConfigTypeCommon, GPUReshapeConfig);},
+			.SetValue = [this](bool bVal) { SetConfigBool(ConfigTypeCommon, GPUReshapeConfig, bVal ); },
+		})
+		.AddBoolean(LOCTEXT("LockDynResLabel","Lock Dynamic Resolution"),
+		{
+			.GetValue = [this]()          { return GetConfigBool(ConfigTypeCommon, LockDynamicResConfig);},
+			.SetValue = [this](bool bVal) { SetConfigBool(ConfigTypeCommon, LockDynamicResConfig, bVal ); },
+		},
+		LOCTEXT("LockDynResHint", "Locks Dynamic Resolution (Default = 60%). This enabled by default if GPU Perf is enabled."))
+		.AddString(LOCTEXT("CVarsLabel", "CVars"), 
+		{
+			.GetValue = [this]() { return GetConfigString(ConfigTypeCommon, CVarConfig); },
+			.SetValue = [this](FString Cvars) { return SetConfigString(ConfigTypeCommon, CVarConfig, Cvars); }
+		},
+		LOCTEXT("CVarsHint", "r.Var1, r.Var2=1, ..."), 
+		LOCTEXT("CVarsToolTip", "Comma separated CVars to be set before test."));
 }
 
 void FAutomatedPerfTestLaunchExtensionInstance::CustomizeUATCommandLine(FString& InOutCommandLine)
@@ -234,6 +278,19 @@ void FAutomatedPerfTestLaunchExtensionInstance::CustomizeUATCommandLine(FString&
 		InOutCommandLine += TEXT(" -AutomatedPerfTest.LockDynamicRes");
 	}
 
+	if (GetConfigBool(ConfigTypeCommon, InsightsConfig))
+	{
+		InOutCommandLine += TEXT(" -AutomatedPerfTest.DoInsightsTrace");
+	}
+
+	// This is enabled by default.
+	if (!GetConfigBool(ConfigTypeCommon, DeterministicConfig, true))
+	{
+		InOutCommandLine += TEXT(" -AutomatedPerfTest.Deterministic=0");
+	}
+
+	InOutCommandLine += TEXT(" -AutomatedPerfTest.IgnoreTestBuildLogging");
+
 	const auto AddStringParamConfig = 
 	[this, &InOutCommandLine](const FString& Config, const FString& Param, const FString& DefaultVal = "")
 	{
@@ -289,6 +346,10 @@ void FAutomatedPerfTestLaunchExtensionInstance::CustomizeUATCommandLine(FString&
 	default: break;
 	}
 
+	AddStringParamConfig(CVarConfig, TEXT("AutomatedPerfTest.CVars"));
+	AddStringParamConfig(TestIDConfig, TEXT("AutomatedPerfTest.TestID"));
+	AddStringParamConfig(InsightsTraceChannelsConfig, TEXT("AutomatedPerfTest.TraceChannels"));
+
 	if (GetConfigBool(ConfigTypeCommon, WindowedConfig))
 	{
 		InOutCommandLine += TEXT(" -windowed");
diff --git a/Engine/Plugins/Performance/AutomatedPerfTesting/Build/Scripts/AutomatedPerfTestConfig.cs b/Engine/Plugins/Performance/AutomatedPerfTesting/Build/Scripts/AutomatedPerfTestConfig.cs
--- a/Engine/Plugins/Performance/AutomatedPerfTesting/Build/Scripts/AutomatedPerfTestConfig.cs
+++ b/Engine/Plugins/Performance/AutomatedPerfTesting/Build/Scripts/AutomatedPerfTestConfig.cs
@@ -116,6 +116,12 @@ namespace AutomatedPerfTest
 		[AutoParamWithNames(false, "AutomatedPerfTest.DoBuildZenStreaming")]
 		public bool DoBuildZenStreaming;
 
+		/// <summary>
+		/// If enabled, runs with -deterministic 
+		/// </summary>
+		[AutoParamWithNames(true, "AutomatedPerfTest.Deterministic")]
+		public bool Deterministic;
+
 		/// <summary>
 		/// The IP/Web Address of the Zen Server
 		/// </summary>
@@ -217,7 +223,7 @@ namespace AutomatedPerfTest
 		/// builds and we do not want to fail the test if no logs are 
 		/// found. 
 		/// </summary>
-		[AutoParamWithNames(false, "AutomatedPerfTest.IgnoreTestBuildLogging")]
+		[AutoParamWithNames(true, "AutomatedPerfTest.IgnoreTestBuildLogging")]
 		public bool IgnoreTestBuildLogging;
 
 		/// <summary>
@@ -250,6 +256,12 @@ namespace AutomatedPerfTest
 		[AutoParamWithNames("", "AutomatedPerfTest.ManagedAccountPool")]
 		public string ManagedAcountPool;
 
+		/// <summary>
+		/// Optional CVars to be passed to client, comma separated. 
+		/// </summary>
+		[AutoParamWithNames("", "AutomatedPerfTest.CVars")]
+		public string CVars;
+
 		/// <summary>
 		/// Overrides default artifact output path by passing the artifact path to the client instance. 
 		/// This is useful if the project has set the artifact to output to a non-standard default output 
diff --git a/Engine/Plugins/Performance/AutomatedPerfTesting/Build/Scripts/AutomatedPerfTestNode.cs b/Engine/Plugins/Performance/AutomatedPerfTesting/Build/Scripts/AutomatedPerfTestNode.cs
--- a/Engine/Plugins/Performance/AutomatedPerfTesting/Build/Scripts/AutomatedPerfTestNode.cs
+++ b/Engine/Plugins/Performance/AutomatedPerfTesting/Build/Scripts/AutomatedPerfTestNode.cs
@@ -974,9 +974,11 @@ namespace AutomatedPerfTest
 
 			if (Config.DoInsightsTrace)
 			{
-				if (Config.TraceChannels != "default,screenshot,stats")
+				const string DefaultTraceChannels = "default,screenshot,stats"; 
+				if (Config.TraceChannels.Trim() != DefaultTraceChannels && GeneratePRSReport)
 				{
-					throw new AutomationException($"Running Insights trace with non-default channels {Config.TraceChannels}, new subtest type is needed to avoid contaminating 'Insights' subtest results.");
+					Log.Warning($"Running Insights trace with non-default channels {Config.TraceChannels}. If this was not intended, it might contaminate 'Insights' subtest results." +
+						$" Change default trace channels to avoid this warning.");
 				}
 
 				return "Insights";
@@ -1020,7 +1022,10 @@ namespace AutomatedPerfTest
 			// force the screen to stay awake during the test
 			ClientRole.CommandLineParams.Add("keepscreenawake");
 			
-			ClientRole.CommandLineParams.Add("deterministic");
+			if(Config.Deterministic)
+			{
+				ClientRole.CommandLineParams.Add("deterministic");
+			}
 
 			Log.Info("AutomatedPerfTestNode<>.GetConfiguration(): Config.DoFPSChart={0}, Config.DoCSVProfiler={1}, Config.DoVideoCapture={2}, Config.DoInsightsTrace={3}, Config.DoLLM={4}", 
 				Config.DoFPSChart, Config.DoCSVProfiler, Config.DoVideoCapture, Config.DoInsightsTrace, Config.DoLLM);
@@ -1173,6 +1178,15 @@ namespace AutomatedPerfTest
 
 			ClientRole.ArchiveDevArtifacts = ForceArchiveDevArtifacts;
 
+			if (!string.IsNullOrEmpty(Config.CVars))
+			{
+				string[] Cvars = Config.CVars.Split(',');
+				foreach (string Cvar in Cvars)
+				{
+					ClientRole.CommandLineParams.AddOrAppendParamValue("dpcvars", Cvar.Trim());
+				}
+			}
+
 			return Config;
 		}
 
